MediaMovieStreamer: No longer requires a player to be passed in.
#rnx

diff --git a/Engine/Plugins/Media/MediaMovieStreamer/Source/MediaMovieStreamer/Private/MediaMovieStreamer.cpp b/Engine/Plugins/Media/MediaMovieStreamer/Source/MediaMovieStreamer/Private/MediaMovieStreamer.cpp
--- a/Engine/Plugins/Media/MediaMovieStreamer/Source/MediaMovieStreamer/Private/MediaMovieStreamer.cpp
+++ b/Engine/Plugins/Media/MediaMovieStreamer/Source/MediaMovieStreamer/Private/MediaMovieStreamer.cpp
@@ -69,26 +69,15 @@ void FMediaMovieStreamer::OnMediaEnd()
 bool FMediaMovieStreamer::Init(const TArray<FString>& InMoviePaths, TEnumAsByte<EMoviePlaybackType> InPlaybackType)
 {
 	MovieViewport->SetTexture(nullptr);
-
-	// Get player.
-	if (MediaPlayer.IsValid() == false)
+	
+	// Play source.
+	if ((MediaPlayer.IsValid()) && (MediaSource.IsValid()))
 	{
-		UE_LOG(LogMediaMovieStreamer, Error, TEXT("OpenNextMovie called but no player set."));
-		return false;
-	}
+		bIsPlaying = true;
+		MediaPlayer->OpenSource(MediaSource.Get());
 
-	// Get source.
-	if (MediaSource.IsValid() == false)
-	{
-		UE_LOG(LogMediaMovieStreamer, Error, TEXT("OpenNextMovie called but no source set."));
-		return false;
+		Texture = MakeShareable(new FSlateTexture2DRHIRef(nullptr, 0, 0));
 	}
-
-	// Play source.
-	bIsPlaying = true;
-	MediaPlayer->OpenSource(MediaSource.Get());
-
-	Texture = MakeShareable(new FSlateTexture2DRHIRef(nullptr, 0, 0));
 	
 	return true;
 }
