Engine/MediaMoviePlayer: adding events for blocking ticks

diff --git a/Engine/Plugins/Media/MediaMovieStreamer/Source/MediaMovieStreamer/Private/MediaMovieStreamer.cpp b/Engine/Plugins/Media/MediaMovieStreamer/Source/MediaMovieStreamer/Private/MediaMovieStreamer.cpp
--- a/Engine/Plugins/Media/MediaMovieStreamer/Source/MediaMovieStreamer/Private/MediaMovieStreamer.cpp
+++ b/Engine/Plugins/Media/MediaMovieStreamer/Source/MediaMovieStreamer/Private/MediaMovieStreamer.cpp
@@ -245,6 +245,10 @@ void FMediaMovieStreamer::TickPreEngine()
 	{
 		MediaModule->TickPreEngine();
 	}
+	if (MovieStreamerPreEngineTick.IsBound())
+	{
+		MovieStreamerPreEngineTick.Broadcast();
+	}
 }
 
 void FMediaMovieStreamer::TickPostEngine()
@@ -254,6 +258,10 @@ void FMediaMovieStreamer::TickPostEngine()
 	{
 		MediaModule->TickPostEngine();
 	}
+	if (MovieStreamerPostEngineTick.IsBound())
+	{
+		MovieStreamerPostEngineTick.Broadcast();
+	}
 }
 
 void FMediaMovieStreamer::TickPostRender()
@@ -263,6 +271,10 @@ void FMediaMovieStreamer::TickPostRender()
 	{
 		MediaModule->TickPostRender();
 	}
+	if (MovieStreamerPostRenderTick.IsBound())
+	{
+		MovieStreamerPostRenderTick.Broadcast();
+	}
 }
 
 IMediaModule* FMediaMovieStreamer::GetMediaModule()
diff --git a/Engine/Plugins/Media/MediaMovieStreamer/Source/MediaMovieStreamer/Public/MediaMovieStreamer.h b/Engine/Plugins/Media/MediaMovieStreamer/Source/MediaMovieStreamer/Public/MediaMovieStreamer.h
--- a/Engine/Plugins/Media/MediaMovieStreamer/Source/MediaMovieStreamer/Public/MediaMovieStreamer.h
+++ b/Engine/Plugins/Media/MediaMovieStreamer/Source/MediaMovieStreamer/Public/MediaMovieStreamer.h
@@ -14,6 +14,8 @@ class UMediaSoundComponent;
 class UMediaSource;
 class UMediaTexture;
 
+DECLARE_EVENT(FMediaMovieStreamer, FMovieStreamerExternalTick);
+
 /**
  * Movie streamer that allows you to use MediaFramework during level loading, etc.
  */
@@ -68,6 +70,10 @@ public:
 	 */
 	MEDIAMOVIESTREAMER_API void SetMediaTexture(UMediaTexture* InMediaTexture);
 
+	FMovieStreamerExternalTick MovieStreamerPostEngineTick;
+	FMovieStreamerExternalTick MovieStreamerPreEngineTick;
+	FMovieStreamerExternalTick MovieStreamerPostRenderTick;
+
 	/**
 	 * Called from UMediaMovieAssets when the media ends.
 	 */
